moved EpicSurvey settings from BaseEngine.ini to BaseEditor.ini

diff --git a/Engine/Config/BaseEditor.ini b/Engine/Config/BaseEditor.ini
--- a/Engine/Config/BaseEditor.ini
+++ b/Engine/Config/BaseEditor.ini
@@ -244,6 +244,15 @@ SupportWebsiteURL="http://unrealengine.com/support"
 [UserDefinedStructure]
 bUseUserDefinedStructure=false
 
+[EpicSurvey]
+Source="CDN"
+CdnUrl="http://cdn.unrealengine.com/surveys/{Culture}/index.json"
+EngineContentRelativeDirectory="Editor/EpicSurvey/{Culture}/"
+NotificationDelayTime=600
+NotificationDuration=10
+PulseTimeInterval=10
+PulseDuration=5
+
 [NewsFeed]
 CdnUrl="http://cdn.unrealengine.com/newsfeed/{Culture}/index.json"
 EngineContentRelativeDirectory="Editor/NewsFeed/{Culture}/"
diff --git a/Engine/Config/BaseEngine.ini b/Engine/Config/BaseEngine.ini
--- a/Engine/Config/BaseEngine.ini
+++ b/Engine/Config/BaseEngine.ini
@@ -1575,15 +1575,6 @@ TEXTUREGROUP_Bokeh=(MinLODSize=1,MaxLODSize=256,LODBias=0,MinMagFilter=linear,Mi
 ; work anyway, so this won't cause any issues.
 ;r.DetailMode=1 
 
-[EpicSurvey]
-Source="CDN"
-CdnUrl="http://cdn.unrealengine.com/surveys/{Culture}/index.json"
-EngineContentRelativeDirectory="Editor/EpicSurvey/{Culture}/"
-NotificationDelayTime=600
-NotificationDuration=10
-PulseTimeInterval=10
-PulseDuration=5
-
 [OnlineSubsystem]
 bHasVoiceEnabled=true
 ; Internal
diff --git a/Engine/Plugins/Editor/EpicSurvey/Source/EpicSurvey/Private/EpicSurvey.cpp b/Engine/Plugins/Editor/EpicSurvey/Source/EpicSurvey/Private/EpicSurvey.cpp
--- a/Engine/Plugins/Editor/EpicSurvey/Source/EpicSurvey/Private/EpicSurvey.cpp
+++ b/Engine/Plugins/Editor/EpicSurvey/Source/EpicSurvey/Private/EpicSurvey.cpp
@@ -109,7 +109,7 @@ void FEpicSurvey::Initialize()
 	{
 		InitializeTitleCloud();
 
-		if ( GConfig->GetInt(TEXT("EpicSurvey"),TEXT("NotificationDelayTime"), SurveyNotificationDelayTime, GEngineIni) )
+		if ( GConfig->GetInt(TEXT("EpicSurvey"),TEXT("NotificationDelayTime"), SurveyNotificationDelayTime, GEditorIni) )
 		{
 			// if the Notification is 0, then it will be displayed during SetActiveSurvey
 			if( SurveyNotificationDelayTime > 0 )
@@ -120,15 +120,15 @@ void FEpicSurvey::Initialize()
 			}
 		}
 
-		GConfig->GetInt(TEXT("EpicSurvey"),TEXT("PulseDuration"), SurveyPulseDuration, GEngineIni);
+		GConfig->GetInt(TEXT("EpicSurvey"),TEXT("PulseDuration"), SurveyPulseDuration, GEditorIni);
 	
-		if ( GConfig->GetInt(TEXT("EpicSurvey"),TEXT("PulseTimeInterval"), SurveyPulseTimeInterval, GEngineIni) )
+		if ( GConfig->GetInt(TEXT("EpicSurvey"),TEXT("PulseTimeInterval"), SurveyPulseTimeInterval, GEditorIni) )
 		{
 			StartPulseSurveyIconDelegate.BindRaw( this, &FEpicSurvey::StartPulseSurveyIcon );
 			EndPulseSurveyIconDelegate.BindRaw( this, &FEpicSurvey::EndPulseSurveyIcon );
 		}
 
-		GConfig->GetFloat(TEXT("EpicSurvey"),TEXT("SurveyNotificationDuration"), SurveyNotificationDuration, GEngineIni);
+		GConfig->GetFloat(TEXT("EpicSurvey"),TEXT("SurveyNotificationDuration"), SurveyNotificationDuration, GEditorIni);
 	}
 
 	if ( TitleCloud.IsValid() )
@@ -164,12 +164,12 @@ void FEpicSurvey::InitializeTitleCloud()
 	};
 
 	FString SourceType;
-	if ( GConfig->GetString(TEXT("EpicSurvey"),TEXT("Source"), SourceType, GEngineIni) && !SourceType.IsEmpty() )
+	if ( GConfig->GetString(TEXT("EpicSurvey"),TEXT("Source"), SourceType, GEditorIni) && !SourceType.IsEmpty() )
 	{
 		if ( SourceType == TEXT("CDN") )
 		{
 			FString CdnUrl;
-			if ( GConfig->GetString(TEXT("EpicSurvey"),TEXT("CdnUrl"), CdnUrl, GEngineIni) && !CdnUrl.IsEmpty() )
+			if ( GConfig->GetString(TEXT("EpicSurvey"),TEXT("CdnUrl"), CdnUrl, GEditorIni) && !CdnUrl.IsEmpty() )
 			{
 				if ( !CultureString.IsEmpty() )
 				{
@@ -186,7 +186,7 @@ void FEpicSurvey::InitializeTitleCloud()
 		else if ( SourceType == TEXT("LOCAL") )
 		{
 			FString EngineContentRelativeDirectory;
-			if ( GConfig->GetString(TEXT("EpicSurvey"),TEXT("EngineContentRelativeDirectory"), EngineContentRelativeDirectory, GEngineIni) && !EngineContentRelativeDirectory.IsEmpty() )
+			if ( GConfig->GetString(TEXT("EpicSurvey"),TEXT("EngineContentRelativeDirectory"), EngineContentRelativeDirectory, GEditorIni) && !EngineContentRelativeDirectory.IsEmpty() )
 			{
 				FString RootDirectory = FPaths::Combine( FPlatformProcess::BaseDir(), *FPaths::EngineContentDir(), *EngineContentRelativeDirectory );
 				if ( !CultureString.IsEmpty() )
