Added return statements to fix build not compiling

diff --git a/Engine/Plugins/Editor/EpicSurvey/Source/EpicSurvey/Private/SurveyTitleCdnStorage.cpp b/Engine/Plugins/Editor/EpicSurvey/Source/EpicSurvey/Private/SurveyTitleCdnStorage.cpp
--- a/Engine/Plugins/Editor/EpicSurvey/Source/EpicSurvey/Private/SurveyTitleCdnStorage.cpp
+++ b/Engine/Plugins/Editor/EpicSurvey/Source/EpicSurvey/Private/SurveyTitleCdnStorage.cpp
@@ -98,7 +98,7 @@ bool FSurveyTitleCdnStorage::EnumerateFiles()
 
 bool FSurveyTitleCdnStorage::EnumerateFiles(int32 Start, int32 Count)
 {
-	EnumerateFiles();
+	return EnumerateFiles();
 }
 
 void FSurveyTitleCdnStorage::GetFileList(TArray<FCloudFileHeader>& OutFiles) 
diff --git a/Engine/Plugins/Editor/EpicSurvey/Source/EpicSurvey/Private/SurveyTitleLocalStorage.cpp b/Engine/Plugins/Editor/EpicSurvey/Source/EpicSurvey/Private/SurveyTitleLocalStorage.cpp
--- a/Engine/Plugins/Editor/EpicSurvey/Source/EpicSurvey/Private/SurveyTitleLocalStorage.cpp
+++ b/Engine/Plugins/Editor/EpicSurvey/Source/EpicSurvey/Private/SurveyTitleLocalStorage.cpp
@@ -87,7 +87,7 @@ bool FSurveyTitleLocalStorage::EnumerateFiles()
 
 bool FSurveyTitleLocalStorage::EnumerateFiles(int32 Start, int32 Count)
 {
-	EnumerateFiles();
+	return EnumerateFiles();
 }
 
 void FSurveyTitleLocalStorage::GetFileList(TArray<FCloudFileHeader>& InFileHeaders)
diff --git a/Engine/Source/Editor/NewsFeed/Private/Implementation/CdnNewsFeedTitleFile.cpp b/Engine/Source/Editor/NewsFeed/Private/Implementation/CdnNewsFeedTitleFile.cpp
--- a/Engine/Source/Editor/NewsFeed/Private/Implementation/CdnNewsFeedTitleFile.cpp
+++ b/Engine/Source/Editor/NewsFeed/Private/Implementation/CdnNewsFeedTitleFile.cpp
@@ -94,7 +94,7 @@ bool FCdnNewsFeedTitleFile::EnumerateFiles()
 
 bool FCdnNewsFeedTitleFile::EnumerateFiles(int32 Start, int32 Count)
 {
-	EnumerateFiles();
+	return EnumerateFiles();
 }
 
 void FCdnNewsFeedTitleFile::GetFileList(TArray<FCloudFileHeader>& OutFiles) 
diff --git a/Engine/Source/Editor/NewsFeed/Private/Implementation/LocalNewsFeedTitleFile.cpp b/Engine/Source/Editor/NewsFeed/Private/Implementation/LocalNewsFeedTitleFile.cpp
--- a/Engine/Source/Editor/NewsFeed/Private/Implementation/LocalNewsFeedTitleFile.cpp
+++ b/Engine/Source/Editor/NewsFeed/Private/Implementation/LocalNewsFeedTitleFile.cpp
@@ -90,7 +90,7 @@ bool FLocalNewsFeedTitleFile::EnumerateFiles()
 
 bool FLocalNewsFeedTitleFile::EnumerateFiles(int32 Start, int32 Count)
 {
-	EnumerateFiles();
+	return EnumerateFiles();
 }
 
 
