Fix crash on startup on Mac due to FSurveyTitleCdnStorage being recreated while iterating over a raw multi-cast delegate.
When the culture file enumeration fails and the FSurveyTitleCdnStorage is recreated in FEpicSurvey::InitializeTitleCloud so using raw-bindings results in a crash when the FSurveyTitleCdnStorage delegate tries to unlock itself. To fix this FSurveyTitleCdnStorage now inherits from TSharedFromThis so that when bound to the HTTP request completion delegate the object is retained.

diff --git a/Engine/Plugins/Editor/EpicSurvey/Source/EpicSurvey/Private/SurveyTitleCdnStorage.cpp b/Engine/Plugins/Editor/EpicSurvey/Source/EpicSurvey/Private/SurveyTitleCdnStorage.cpp
--- a/Engine/Plugins/Editor/EpicSurvey/Source/EpicSurvey/Private/SurveyTitleCdnStorage.cpp
+++ b/Engine/Plugins/Editor/EpicSurvey/Source/EpicSurvey/Private/SurveyTitleCdnStorage.cpp
@@ -88,7 +88,7 @@ bool FSurveyTitleCdnStorage::EnumerateFiles(const FPagedQuery& Page)
 	TSharedRef<class IHttpRequest> HttpRequest = FHttpModule::Get().CreateRequest();
 	EnumerateFilesRequests.Enqueue(&HttpRequest.Get());
 
-	HttpRequest->OnProcessRequestComplete().BindRaw(this, &FSurveyTitleCdnStorage::EnumerateFiles_HttpRequestComplete);
+	HttpRequest->OnProcessRequestComplete().BindThreadSafeSP(AsShared(), &FSurveyTitleCdnStorage::EnumerateFiles_HttpRequestComplete);
 	HttpRequest->SetURL( IndexUrl );
 	HttpRequest->SetVerb(TEXT("GET"));
 	HttpRequest->ProcessRequest();
@@ -244,7 +244,7 @@ bool FSurveyTitleCdnStorage::ReadFile(const FString& FileName)
 	TSharedRef<class IHttpRequest> HttpRequest = FHttpModule::Get().CreateRequest();
 	FileRequests.Add(&HttpRequest.Get(), FPendingFileRequest(FileName));
 
-	HttpRequest->OnProcessRequestComplete().BindRaw(this, &FSurveyTitleCdnStorage::ReadFile_HttpRequestComplete);
+	HttpRequest->OnProcessRequestComplete().BindThreadSafeSP(AsShared(), &FSurveyTitleCdnStorage::ReadFile_HttpRequestComplete);
 	HttpRequest->SetURL( FileName );
 	HttpRequest->SetVerb(TEXT("GET"));
 
diff --git a/Engine/Plugins/Editor/EpicSurvey/Source/EpicSurvey/Private/SurveyTitleCdnStorage.h b/Engine/Plugins/Editor/EpicSurvey/Source/EpicSurvey/Private/SurveyTitleCdnStorage.h
--- a/Engine/Plugins/Editor/EpicSurvey/Source/EpicSurvey/Private/SurveyTitleCdnStorage.h
+++ b/Engine/Plugins/Editor/EpicSurvey/Source/EpicSurvey/Private/SurveyTitleCdnStorage.h
@@ -7,7 +7,7 @@
 #include "IHttpRequest.h"
 #include "IHttpResponse.h"
 
-class FSurveyTitleCdnStorage : public IOnlineTitleFile
+class FSurveyTitleCdnStorage : public TSharedFromThis<FSurveyTitleCdnStorage, ESPMode::ThreadSafe>, public IOnlineTitleFile
 {
 public:
 	static IOnlineTitleFilePtr Create( const FString& IndexUrl );
