Epic Survey - Fix toast display duration and toolbar icon pulsing.
+ Update the toast duration setting to match what's defined in the ini files.
+ Initialize the title cloud after reading in the settings since they are used in the initialization.
+ Simplify the pulse logic so that it works again.
+ Adjust the pulse animation timing so that the icon's opacity doesn't pop, and don't let it's opacity go all the way to 0.
+ Add logging for when the JSON survey file can't be parsed.

TTP# 341104
TTP# 341105

diff --git a/Engine/Config/BaseEditor.ini b/Engine/Config/BaseEditor.ini
--- a/Engine/Config/BaseEditor.ini
+++ b/Engine/Config/BaseEditor.ini
@@ -266,7 +266,7 @@ EngineContentRelativeDirectory="Editor/EpicSurvey/{Culture}/"
 NotificationDelayTime=600
 NotificationDuration=10
 PulseTimeInterval=10
-PulseDuration=5
+PulseDuration=6
 
 [CustomBlueprintFunctionLibrary]
 bCanCreateNew=true
diff --git a/Engine/Plugins/Editor/EpicSurvey/Source/EpicSurvey/Private/EpicSurvey.cpp b/Engine/Plugins/Editor/EpicSurvey/Source/EpicSurvey/Private/EpicSurvey.cpp
--- a/Engine/Plugins/Editor/EpicSurvey/Source/EpicSurvey/Private/EpicSurvey.cpp
+++ b/Engine/Plugins/Editor/EpicSurvey/Source/EpicSurvey/Private/EpicSurvey.cpp
@@ -18,6 +18,8 @@
 
 #define LOCTEXT_NAMESPACE "EpicSurvey"
 
+DEFINE_LOG_CATEGORY_STATIC(LogEpicSurvey, Display, All);
+
 const FString FEpicSurvey::SurveyIndexFilename( TEXT("survey_index.json") );
 
 TSharedRef< FEpicSurvey > FEpicSurvey::Create()
@@ -107,8 +109,6 @@ void FEpicSurvey::Initialize()
 
 	if ( !IsRunningCommandlet() )
 	{
-		InitializeTitleCloud();
-
 		if ( GConfig->GetInt(TEXT("EpicSurvey"),TEXT("NotificationDelayTime"), SurveyNotificationDelayTime, GEditorIni) )
 		{
 			// if the Notification is 0, then it will be displayed during SetActiveSurvey
@@ -121,14 +121,16 @@ void FEpicSurvey::Initialize()
 		}
 
 		GConfig->GetInt(TEXT("EpicSurvey"),TEXT("PulseDuration"), SurveyPulseDuration, GEditorIni);
-	
+
 		if ( GConfig->GetInt(TEXT("EpicSurvey"),TEXT("PulseTimeInterval"), SurveyPulseTimeInterval, GEditorIni) )
 		{
 			StartPulseSurveyIconDelegate.BindRaw( this, &FEpicSurvey::StartPulseSurveyIcon );
 			EndPulseSurveyIconDelegate.BindRaw( this, &FEpicSurvey::EndPulseSurveyIcon );
 		}
 
-		GConfig->GetFloat(TEXT("EpicSurvey"),TEXT("SurveyNotificationDuration"), SurveyNotificationDuration, GEditorIni);
+		GConfig->GetFloat(TEXT("EpicSurvey"),TEXT("NotificationDuration"), SurveyNotificationDuration, GEditorIni);
+
+		InitializeTitleCloud();
 	}
 
 	if ( TitleCloud.IsValid() )
@@ -335,6 +337,10 @@ void FEpicSurvey::OnReadFileComplete( bool bSuccess, const FString& DLName )
 							}
 						}
 					}
+					else
+					{
+						UE_LOG(LogEpicSurvey, Warning, TEXT("Parsing JSON survey failed. Filename: %s Message: %s"), *FileHeader.FileName, *Reader->GetErrorMessage());
+					}
 				}
 				else if ( FileExtension == TEXT("png") )
 				{
@@ -382,8 +388,7 @@ void FEpicSurvey::DisplayToolbarNotification()
 
 		if( StartPulseSurveyIconDelegate.IsBound() && EndPulseSurveyIconDelegate.IsBound() )
 		{
-			bSurveyIconPulsing = true;
-			GEditor->GetTimerManager()->SetTimer( StartPulseSurveyIconDelegate, SurveyPulseTimeInterval, false );
+			StartPulseSurveyIcon();
 		}
 	}
 }
@@ -443,18 +448,14 @@ void FEpicSurvey::CancelSurveyNotification()
 
 void FEpicSurvey::StartPulseSurveyIcon()
 {
-	GEditor->GetTimerManager()->SetTimer( StartPulseSurveyIconDelegate, 0, false );
 	GEditor->GetTimerManager()->SetTimer( EndPulseSurveyIconDelegate, SurveyPulseDuration, false );
-
 	bSurveyIconPulsing = true;
 }
 
 void FEpicSurvey::EndPulseSurveyIcon()
 {
 	bSurveyIconPulsing = false;
-
 	GEditor->GetTimerManager()->SetTimer( StartPulseSurveyIconDelegate, SurveyPulseTimeInterval, false );
-	GEditor->GetTimerManager()->SetTimer( EndPulseSurveyIconDelegate, 0, false );
 }
 
 bool FEpicSurvey::PromptSurvey( const FGuid& SurveyIdentifier )
diff --git a/Engine/Plugins/Editor/EpicSurvey/Source/EpicSurvey/Private/SSurveyNotification.h b/Engine/Plugins/Editor/EpicSurvey/Source/EpicSurvey/Private/SSurveyNotification.h
--- a/Engine/Plugins/Editor/EpicSurvey/Source/EpicSurvey/Private/SSurveyNotification.h
+++ b/Engine/Plugins/Editor/EpicSurvey/Source/EpicSurvey/Private/SSurveyNotification.h
@@ -18,7 +18,7 @@ public:
 		Survey = InSurvey;
 
 		Color = FSlateColor( FLinearColor( 1, 1, 1, 0 ) );
-		Sequence.AddCurve( 0, 1.0f, ECurveEaseFunction::QuadIn );
+		Sequence.AddCurve( 0, 1, ECurveEaseFunction::QuadOut );
 
 		ChildSlot
 		[
@@ -48,7 +48,17 @@ public:
 			float Opacity = 1.0f;
 			if ( EpicSurvey->ShouldPulseSurveyIcon( Survey.ToSharedRef() ) )
 			{
-				Opacity = Sequence.GetLerp();
+				// When pulsing starts jump to the start of the animation.
+				if (!bIsPulsing)
+				{
+					const_cast<SSurveyNotification*>(this)->Sequence.JumpToStart();
+					bIsPulsing = true;
+				}
+				Opacity = 1 - (Sequence.GetLerp() * .8f);
+			}
+			else
+			{
+				bIsPulsing = false;
 			}
 
 			if ( !Sequence.IsPlaying() )
@@ -95,6 +105,7 @@ private:
 	/** Animation sequence to pulse the image */
 	FCurveSequence Sequence;
 	FSlateColor Color;
+	bool bIsPulsing;
 
 };
 
diff --git a/Engine/Plugins/Editor/EpicSurvey/Source/EpicSurvey/Private/SurveyTitleCdnStorage.cpp b/Engine/Plugins/Editor/EpicSurvey/Source/EpicSurvey/Private/SurveyTitleCdnStorage.cpp
--- a/Engine/Plugins/Editor/EpicSurvey/Source/EpicSurvey/Private/SurveyTitleCdnStorage.cpp
+++ b/Engine/Plugins/Editor/EpicSurvey/Source/EpicSurvey/Private/SurveyTitleCdnStorage.cpp
@@ -7,8 +7,6 @@
 #include "SecureHash.h"
 #include "HttpModule.h"
 
-DEFINE_LOG_CATEGORY_STATIC(LogEpicSurvey, Display, All);
-
 IOnlineTitleFilePtr FSurveyTitleCdnStorage::Create( const FString& IndexUrl )
 {
 	return MakeShareable( new FSurveyTitleCdnStorage( IndexUrl ) );
