EpicSurvey - Fix the non-unity build.

This doesn't seem like the correct fix, but it's the only one that makes the unity and non-unity builds succeede locally.  I'll look at it again next week.

diff --git a/Engine/Plugins/Editor/EpicSurvey/Source/EpicSurvey/Private/EpicSurvey.cpp b/Engine/Plugins/Editor/EpicSurvey/Source/EpicSurvey/Private/EpicSurvey.cpp
--- a/Engine/Plugins/Editor/EpicSurvey/Source/EpicSurvey/Private/EpicSurvey.cpp
+++ b/Engine/Plugins/Editor/EpicSurvey/Source/EpicSurvey/Private/EpicSurvey.cpp
@@ -18,7 +18,7 @@
 
 #define LOCTEXT_NAMESPACE "EpicSurvey"
 
-DEFINE_LOG_CATEGORY_STATIC(LogEpicSurvey, Display, All);
+DEFINE_LOG_CATEGORY(LogEpicSurvey);
 
 const FString FEpicSurvey::SurveyIndexFilename( TEXT("survey_index.json") );
 
diff --git a/Engine/Plugins/Editor/EpicSurvey/Source/EpicSurvey/Private/SurveyTitleCdnStorage.h b/Engine/Plugins/Editor/EpicSurvey/Source/EpicSurvey/Private/SurveyTitleCdnStorage.h
--- a/Engine/Plugins/Editor/EpicSurvey/Source/EpicSurvey/Private/SurveyTitleCdnStorage.h
+++ b/Engine/Plugins/Editor/EpicSurvey/Source/EpicSurvey/Private/SurveyTitleCdnStorage.h
@@ -7,6 +7,8 @@
 #include "IHttpRequest.h"
 #include "IHttpResponse.h"
 
+DECLARE_LOG_CATEGORY_EXTERN(LogEpicSurvey, Display, All);
+
 class FSurveyTitleCdnStorage : public TSharedFromThis<FSurveyTitleCdnStorage, ESPMode::ThreadSafe>, public IOnlineTitleFile
 {
 public:
