Made survey warnings verbose to eliminate log spam

diff --git a/Engine/Plugins/Editor/EpicSurvey/Source/EpicSurvey/Private/EpicSurvey.cpp b/Engine/Plugins/Editor/EpicSurvey/Source/EpicSurvey/Private/EpicSurvey.cpp
--- a/Engine/Plugins/Editor/EpicSurvey/Source/EpicSurvey/Private/EpicSurvey.cpp
+++ b/Engine/Plugins/Editor/EpicSurvey/Source/EpicSurvey/Private/EpicSurvey.cpp
@@ -339,7 +339,7 @@ void FEpicSurvey::OnReadFileComplete( bool bSuccess, const FString& DLName )
 					}
 					else
 					{
-						UE_LOG(LogEpicSurvey, Warning, TEXT("Parsing JSON survey failed. Filename: %s Message: %s"), *FileHeader.FileName, *Reader->GetErrorMessage());
+						UE_LOG(LogEpicSurvey, Verbose, TEXT("Parsing JSON survey failed. Filename: %s Message: %s"), *FileHeader.FileName, *Reader->GetErrorMessage());
 					}
 				}
 				else if ( FileExtension == TEXT("png") )
diff --git a/Engine/Plugins/Editor/EpicSurvey/Source/EpicSurvey/Private/SurveyTitleCdnStorage.cpp b/Engine/Plugins/Editor/EpicSurvey/Source/EpicSurvey/Private/SurveyTitleCdnStorage.cpp
--- a/Engine/Plugins/Editor/EpicSurvey/Source/EpicSurvey/Private/SurveyTitleCdnStorage.cpp
+++ b/Engine/Plugins/Editor/EpicSurvey/Source/EpicSurvey/Private/SurveyTitleCdnStorage.cpp
@@ -41,7 +41,7 @@ bool FSurveyTitleCdnStorage::ClearFiles()
 	{
 		if (Files[Idx].AsyncState == EOnlineAsyncTaskState::InProgress)
 		{
-			UE_LOG(LogEpicSurvey, Warning, TEXT("Cant clear files. Pending file op for %s"), *Files[Idx].FileName);
+			UE_LOG(LogEpicSurvey, Verbose, TEXT("Cant clear files. Pending file op for %s"), *Files[Idx].FileName);
 			return false;
 		}
 	}
@@ -59,7 +59,7 @@ bool FSurveyTitleCdnStorage::ClearFile(const FString& FileName)
 		{
 			if (Files[Idx].AsyncState == EOnlineAsyncTaskState::InProgress)
 			{
-				UE_LOG(LogEpicSurvey, Warning, TEXT("Cant clear file. Pending file op for %s"), *Files[Idx].FileName);
+				UE_LOG(LogEpicSurvey, Verbose, TEXT("Cant clear file. Pending file op for %s"), *Files[Idx].FileName);
 				return false;
 			}
 			else
@@ -77,7 +77,7 @@ bool FSurveyTitleCdnStorage::EnumerateFiles(const FPagedQuery& Page)
 	// Make sure an enumeration request  is not currently pending
 	if(!EnumerateFilesRequests.IsEmpty())
 	{
-		UE_LOG(LogEpicSurvey, Warning, TEXT("EnumerateFiles request failed. Request already in progress."));
+		UE_LOG(LogEpicSurvey, Verbose, TEXT("EnumerateFiles request failed. Request already in progress."));
 		TriggerOnEnumerateFilesCompleteDelegates(false);
 		return false;
 	}
@@ -142,7 +142,7 @@ void FSurveyTitleCdnStorage::SaveCloudFileToDisk(const FString& Filename, const
 	}
 	else
 	{
-		UE_LOG(LogEpicSurvey, Warning, TEXT("WriteUserFile request complete. Local file cache failed to update =%s"), 
+		UE_LOG(LogEpicSurvey, Verbose, TEXT("WriteUserFile request complete. Local file cache failed to update =%s"), 
 			*LocalFilePath);
 	}
 }
@@ -185,7 +185,7 @@ bool FSurveyTitleCdnStorage::ReadFile(const FString& FileName)
 	
 	if (!ErrorStr.IsEmpty())
 	{
-		UE_LOG(LogEpicSurvey, Warning, TEXT("ReadFile request failed. %s"), *ErrorStr);
+		UE_LOG(LogEpicSurvey, Verbose, TEXT("ReadFile request failed. %s"), *ErrorStr);
 		TriggerOnReadFileCompleteDelegates(false, FileName);
 		return false;
 	}
@@ -315,7 +315,7 @@ void FSurveyTitleCdnStorage::EnumerateFiles_HttpRequestComplete(FHttpRequestPtr
 						if (FileHeader.Hash.IsEmpty() ||
 							FileHeader.DLName.IsEmpty())
 						{
-							UE_LOG(LogEpicSurvey, Warning, TEXT("Invalid file entry hash=%s dlname=%s filename=%s"),
+							UE_LOG(LogEpicSurvey, Verbose, TEXT("Invalid file entry hash=%s dlname=%s filename=%s"),
 								*FileHeader.Hash, *FileHeader.DLName, *FileHeader.FileName);
 						}
 						else
@@ -343,7 +343,7 @@ void FSurveyTitleCdnStorage::EnumerateFiles_HttpRequestComplete(FHttpRequestPtr
 
 	if (!ErrorStr.IsEmpty())
 	{
-		UE_LOG(LogEpicSurvey, Warning, TEXT("EnumerateFiles request failed. %s"), *ErrorStr);
+		UE_LOG(LogEpicSurvey, Verbose, TEXT("EnumerateFiles request failed. %s"), *ErrorStr);
 	}
 
 	TriggerOnEnumerateFilesCompleteDelegates(bResult);
@@ -395,7 +395,7 @@ void FSurveyTitleCdnStorage::ReadFile_HttpRequestComplete(FHttpRequestPtr HttpRe
 	}
 	if (!ErrorStr.IsEmpty())
 	{
-		UE_LOG(LogEpicSurvey, Warning, TEXT("ReadFile request failed. %s"), *ErrorStr);
+		UE_LOG(LogEpicSurvey, Verbose, TEXT("ReadFile request failed. %s"), *ErrorStr);
 	}
 
 	TriggerOnReadFileCompleteDelegates(bResult, PendingRequest.FileName);
