PR #968: Fixing PVS-Studio warnings (Contributed by [REDACTED_PERSON])

These can be generally described as removing redundant condition checks, using logical operators instead of bitwise operators, removal of redundant parenthesis, and consistent use of pragma push and pragma pop. There is very little observable behavior change in this CL.

Changes that required further review are pending in 2489753

diff --git a/Engine/Plugins/Editor/EpicSurvey/Source/EpicSurvey/Private/Survey.cpp b/Engine/Plugins/Editor/EpicSurvey/Source/EpicSurvey/Private/Survey.cpp
--- a/Engine/Plugins/Editor/EpicSurvey/Source/EpicSurvey/Private/Survey.cpp
+++ b/Engine/Plugins/Editor/EpicSurvey/Source/EpicSurvey/Private/Survey.cpp
@@ -258,30 +258,27 @@ void FSurvey::EvaluateBranches()
 		TSharedPtr< FSurvey > OldBranchSurvey = Pages[NextPageIndex]->GetBranchSurvey();
 		if( OldBranchSurvey.IsValid() && (!NewBranch.IsValid() || (NewBranch->GetBranchSurvey() != OldBranchSurvey)) )
 		{
-			if( OldBranchSurvey.IsValid() )
-			{
-				// Mark the branch as not completed.
-				OldBranchSurvey->SetBranchUsed( false );
+			// Mark the branch as not completed.
+			OldBranchSurvey->SetBranchUsed( false );
 
-				// Remove the branch survey pages from the survey.
-				struct FRemovePred
+			// Remove the branch survey pages from the survey.
+			struct FRemovePred
+			{
+				FRemovePred( TSharedPtr< FSurvey > InOldBranchSurvey )
+					: Pages(InOldBranchSurvey->GetPages())
 				{
-					FRemovePred( TSharedPtr< FSurvey > InOldBranchSurvey )
-						: Pages(InOldBranchSurvey->GetPages())
-					{
-					}
-					bool operator()( const TSharedPtr<FSurveyPage>& InPage ) const
-					{
-						return Pages.Contains( InPage.ToSharedRef() );
-					}
-				private:
-					const TArray< TSharedRef< FSurveyPage > >& Pages;
-				};
-				Pages.RemoveAll( FRemovePred(OldBranchSurvey) );
+				}
+				bool operator()( const TSharedPtr<FSurveyPage>& InPage ) const
+				{
+					return Pages.Contains( InPage.ToSharedRef() );
+				}
+			private:
+				const TArray< TSharedRef< FSurveyPage > >& Pages;
+			};
+			Pages.RemoveAll( FRemovePred(OldBranchSurvey) );
 
-				// remove the branch points for any of the branch survey answers.
-				OldBranchSurvey->UpdateAllBranchPoints( false );
-			}
+			// remove the branch points for any of the branch survey answers.
+			OldBranchSurvey->UpdateAllBranchPoints( false );
 		}
 
 	}
@@ -333,10 +330,7 @@ void FSurvey::OnPageBack()
 		TSharedPtr< FSurvey > BranchSurvey = Pages[CurrentPageIndex-1]->GetBranchSurvey();
 		if( BranchSurvey.IsValid() )
 		{
-			if( BranchSurvey.IsValid() )
-			{
-				BranchSurvey->UpdateAllBranchPoints( false );
-			}
+			BranchSurvey->UpdateAllBranchPoints( false );
 		}
 		else
 		{
diff --git a/Engine/Plugins/Experimental/Phya/Source/Phya/Private/PhyaLib/src/Resonator/paModalData.cpp b/Engine/Plugins/Experimental/Phya/Source/Phya/Private/PhyaLib/src/Resonator/paModalData.cpp
--- a/Engine/Plugins/Experimental/Phya/Source/Phya/Private/PhyaLib/src/Resonator/paModalData.cpp
+++ b/Engine/Plugins/Experimental/Phya/Source/Phya/Private/PhyaLib/src/Resonator/paModalData.cpp
@@ -38,6 +38,7 @@ paModalData :: ~paModalData() {
 }
 
 #if PLATFORM_WINDOWS
+#pragma warning (push)
 #pragma warning (disable : 4996) // scanf is deprecated, but can't use the safe version as may not be supported on all platforms.
 #endif
 
@@ -86,5 +87,5 @@ paModalData :: read( const char* data ) {
 }
 
 #if PLATFORM_WINDOWS
-#pragma warning (default : 4996) // scanf is deprecated.
+#pragma warning (pop) // scanf is deprecated.
 #endif
diff --git a/Engine/Plugins/Runtime/OculusRift/Source/OculusRift/Private/OculusRiftHMD.cpp b/Engine/Plugins/Runtime/OculusRift/Source/OculusRift/Private/OculusRiftHMD.cpp
--- a/Engine/Plugins/Runtime/OculusRift/Source/OculusRift/Private/OculusRiftHMD.cpp
+++ b/Engine/Plugins/Runtime/OculusRift/Source/OculusRift/Private/OculusRiftHMD.cpp
@@ -1250,7 +1250,7 @@ bool FOculusRiftHMD::DoEnableStereo(bool bStereo, bool bApplyToHmd)
 					SceneVP->SetViewportSize(Hmd->Resolution.w, Hmd->Resolution.h);
 				}
 			}
-			else if ((!wasFullscreenAllowed && !stereoEnabled))
+			else if (!wasFullscreenAllowed && !stereoEnabled)
 			{
 				// restoring original viewport size (to be equal to window size).
 				TSharedPtr<SWindow> Window = SceneVP->FindWindow();
diff --git a/Engine/Plugins/ScriptGeneratorPlugin/Source/ScriptGeneratorPlugin/Private/GenericScriptCodeGenerator.cpp b/Engine/Plugins/ScriptGeneratorPlugin/Source/ScriptGeneratorPlugin/Private/GenericScriptCodeGenerator.cpp
--- a/Engine/Plugins/ScriptGeneratorPlugin/Source/ScriptGeneratorPlugin/Private/GenericScriptCodeGenerator.cpp
+++ b/Engine/Plugins/ScriptGeneratorPlugin/Source/ScriptGeneratorPlugin/Private/GenericScriptCodeGenerator.cpp
@@ -182,33 +182,30 @@ bool FGenericScriptCodeGenerator::CanExportClass(UClass* Class)
 	bool bCanExport = FScriptCodeGeneratorBase::CanExportClass(Class);
 	if (bCanExport)
 	{
-		if (bCanExport)
+		const FString ClassNameCPP = GetClassNameCPP(Class);
+		// No members to export? Don't bother exporting the class.
+		bool bHasMembersToExport = false;
+		for (TFieldIterator<UFunction> FuncIt(Class, EFieldIteratorFlags::ExcludeSuper); !bHasMembersToExport && FuncIt; ++FuncIt)
 		{
-			const FString ClassNameCPP = GetClassNameCPP(Class);
-			// No members to export? Don't bother exporting the class.
-			bool bHasMembersToExport = false;
-			for (TFieldIterator<UFunction> FuncIt(Class, EFieldIteratorFlags::ExcludeSuper); !bHasMembersToExport && FuncIt; ++FuncIt)
+			UFunction* Function = *FuncIt;
+			if (CanExportFunction(ClassNameCPP, Class, Function))
 			{
-				UFunction* Function = *FuncIt;
-				if (CanExportFunction(ClassNameCPP, Class, Function))
-				{
-					bHasMembersToExport = true;
-				}
+				bHasMembersToExport = true;
 			}
-			// Check properties too
-			if (!bHasMembersToExport)
+		}
+		// Check properties too
+		if (!bHasMembersToExport)
+		{
+			for (TFieldIterator<UProperty> PropertyIt(Class, EFieldIteratorFlags::ExcludeSuper); !bHasMembersToExport && PropertyIt; ++PropertyIt)
 			{
-				for (TFieldIterator<UProperty> PropertyIt(Class, EFieldIteratorFlags::ExcludeSuper); !bHasMembersToExport && PropertyIt; ++PropertyIt)
+				UProperty* Property = *PropertyIt;
+				if (CanExportProperty(ClassNameCPP, Class, Property))
 				{
-					UProperty* Property = *PropertyIt;
-					if (CanExportProperty(ClassNameCPP, Class, Property))
-					{
-						bHasMembersToExport = true;
-					}
+					bHasMembersToExport = true;
 				}
 			}
-			bCanExport = bHasMembersToExport;
 		}
+		bCanExport = bHasMembersToExport;
 	}
 	return bCanExport;
 }
diff --git a/Engine/Plugins/ScriptGeneratorPlugin/Source/ScriptGeneratorPlugin/Private/LuaScriptCodeGenerator.cpp b/Engine/Plugins/ScriptGeneratorPlugin/Source/ScriptGeneratorPlugin/Private/LuaScriptCodeGenerator.cpp
--- a/Engine/Plugins/ScriptGeneratorPlugin/Source/ScriptGeneratorPlugin/Private/LuaScriptCodeGenerator.cpp
+++ b/Engine/Plugins/ScriptGeneratorPlugin/Source/ScriptGeneratorPlugin/Private/LuaScriptCodeGenerator.cpp
@@ -199,33 +199,30 @@ bool FLuaScriptCodeGenerator::CanExportClass(UClass* Class)
 	bool bCanExport = FScriptCodeGeneratorBase::CanExportClass(Class);
 	if (bCanExport)
 	{
-		if (bCanExport)
+		const FString ClassNameCPP = GetClassNameCPP(Class);
+		// No functions to export? Don't bother exporting the class.
+		bool bHasMembersToExport = false;
+		for (TFieldIterator<UFunction> FuncIt(Class); !bHasMembersToExport && FuncIt; ++FuncIt)
 		{
-			const FString ClassNameCPP = GetClassNameCPP(Class);
-			// No functions to export? Don't bother exporting the class.
-			bool bHasMembersToExport = false;
-			for (TFieldIterator<UFunction> FuncIt(Class); !bHasMembersToExport && FuncIt; ++FuncIt)
+			UFunction* Function = *FuncIt;
+			if (CanExportFunction(ClassNameCPP, Class, Function))
 			{
-				UFunction* Function = *FuncIt;
-				if (CanExportFunction(ClassNameCPP, Class, Function))
-				{
-					bHasMembersToExport = true;
-				}
-			}			
-			// Check properties too
-			if (!bHasMembersToExport)
+				bHasMembersToExport = true;
+			}
+		}			
+		// Check properties too
+		if (!bHasMembersToExport)
+		{
+			for (TFieldIterator<UProperty> PropertyIt(Class, EFieldIteratorFlags::ExcludeSuper); !bHasMembersToExport && PropertyIt; ++PropertyIt)
 			{
-				for (TFieldIterator<UProperty> PropertyIt(Class, EFieldIteratorFlags::ExcludeSuper); !bHasMembersToExport && PropertyIt; ++PropertyIt)
+				UProperty* Property = *PropertyIt;
+				if (CanExportProperty(ClassNameCPP, Class, Property))
 				{
-					UProperty* Property = *PropertyIt;
-					if (CanExportProperty(ClassNameCPP, Class, Property))
-					{
-						bHasMembersToExport = true;
-					}
+					bHasMembersToExport = true;
 				}
 			}
-			bCanExport = bHasMembersToExport;
 		}
+		bCanExport = bHasMembersToExport;
 	}
 	return bCanExport;
 }
diff --git a/Engine/Source/Developer/AssetTools/Private/AssetTypeActions/AssetTypeActions_Blueprint.cpp b/Engine/Source/Developer/AssetTools/Private/AssetTypeActions/AssetTypeActions_Blueprint.cpp
--- a/Engine/Source/Developer/AssetTools/Private/AssetTypeActions/AssetTypeActions_Blueprint.cpp
+++ b/Engine/Source/Developer/AssetTools/Private/AssetTypeActions/AssetTypeActions_Blueprint.cpp
@@ -134,7 +134,7 @@ void FAssetTypeActions_Blueprint::ExecuteEditDefaults(TArray<TWeakObjectPtr<UBlu
 		if ( Object )
 		{
 			// If the blueprint is valid, allow it to be added to the list, otherwise log the error.
-			if (Object && Object->SkeletonGeneratedClass && Object->GeneratedClass )
+			if ( Object->SkeletonGeneratedClass && Object->GeneratedClass )
 			{
 				Blueprints.Add(Object);
 			}
diff --git a/Engine/Source/Developer/MeshUtilities/Private/MeshUtilities.cpp b/Engine/Source/Developer/MeshUtilities/Private/MeshUtilities.cpp
--- a/Engine/Source/Developer/MeshUtilities/Private/MeshUtilities.cpp
+++ b/Engine/Source/Developer/MeshUtilities/Private/MeshUtilities.cpp
@@ -4016,7 +4016,7 @@ bool FMeshUtilities::ConstructRawMesh(
 	}
 
 	// Compute any missing tangents.
-	if (bRecomputeTangents || bRecomputeTangents)
+	if (bRecomputeTangents)
 	{
 		float ComparisonThreshold = GetComparisonThreshold(BuildSettings);
 		TMultiMap<int32,int32> OverlappingCorners;
diff --git a/Engine/Source/Developer/TextureCompressor/Private/TextureCompressorModule.cpp b/Engine/Source/Developer/TextureCompressor/Private/TextureCompressorModule.cpp
--- a/Engine/Source/Developer/TextureCompressor/Private/TextureCompressorModule.cpp
+++ b/Engine/Source/Developer/TextureCompressor/Private/TextureCompressorModule.cpp
@@ -1779,7 +1779,6 @@ private:
 		int32 NumSourceMips = InSourceMips.Num();
 
 		if (BuildSettings.MipGenSettings != TMGS_LeaveExistingMips ||
-			BuildSettings.MipGenSettings == TMGS_NoMipmaps ||
 			bLongLatCubemap)
 		{
 			NumSourceMips = 1;
diff --git a/Engine/Source/Editor/BlueprintGraph/Private/BlueprintGraphPanelPinFactory.cpp b/Engine/Source/Editor/BlueprintGraph/Private/BlueprintGraphPanelPinFactory.cpp
--- a/Engine/Source/Editor/BlueprintGraph/Private/BlueprintGraphPanelPinFactory.cpp
+++ b/Engine/Source/Editor/BlueprintGraph/Private/BlueprintGraphPanelPinFactory.cpp
@@ -36,7 +36,7 @@ TSharedPtr<class SGraphPin> FBlueprintGraphPanelPinFactory::CreatePin(class UEdG
 
 		if (DataTablePin)
 		{
-			if (DataTablePin && DataTablePin->DefaultObject != nullptr && DataTablePin->LinkedTo.Num() == 0)
+			if (DataTablePin->DefaultObject != nullptr && DataTablePin->LinkedTo.Num() == 0)
 			{
 				if (auto DataTable = Cast<UDataTable>(DataTablePin->DefaultObject))
 				{
diff --git a/Engine/Source/Editor/EditorWidgets/Private/SObjectNameEditableTextBox.cpp b/Engine/Source/Editor/EditorWidgets/Private/SObjectNameEditableTextBox.cpp
--- a/Engine/Source/Editor/EditorWidgets/Private/SObjectNameEditableTextBox.cpp
+++ b/Engine/Source/Editor/EditorWidgets/Private/SObjectNameEditableTextBox.cpp
@@ -78,8 +78,6 @@ int32 SObjectNameEditableTextBox::OnPaint( const FPaintArgs& Args, const FGeomet
 
 	const double CurrentTime = FSlateApplication::Get().GetCurrentTime();
 
-	const float DrawPositionY = ( AllottedGeometry.Size.Y / 2 ) - ( AllottedGeometry.Size.Y / 2 );
-
 	// Draw highlight targeting effect
 	const float TimeSinceHighlightInteraction = (float)( CurrentTime - LastCommittedTime );
 	if( TimeSinceHighlightInteraction <= HighlightTargetEffectDuration )
diff --git a/Engine/Source/Editor/LandscapeEditor/Private/LandscapeEditorDetailCustomization_NewLandscape.cpp b/Engine/Source/Editor/LandscapeEditor/Private/LandscapeEditorDetailCustomization_NewLandscape.cpp
--- a/Engine/Source/Editor/LandscapeEditor/Private/LandscapeEditorDetailCustomization_NewLandscape.cpp
+++ b/Engine/Source/Editor/LandscapeEditor/Private/LandscapeEditorDetailCustomization_NewLandscape.cpp
@@ -1029,10 +1029,7 @@ FText FLandscapeEditorDetailCustomization_NewLandscape::GetImportHeightmapFilena
 
 	if (Result == FPropertyAccess::MultipleValues)
 	{
-		if (Result == FPropertyAccess::MultipleValues)
-		{
-			return NSLOCTEXT("PropertyEditor", "MultipleValues", "Multiple Values");
-		}
+		return NSLOCTEXT("PropertyEditor", "MultipleValues", "Multiple Values");
 	}
 
 	return FText::FromString(HeightmapFilename);
diff --git a/Engine/Source/Editor/PropertyEditor/Private/SDetailSingleItemRow.cpp b/Engine/Source/Editor/PropertyEditor/Private/SDetailSingleItemRow.cpp
--- a/Engine/Source/Editor/PropertyEditor/Private/SDetailSingleItemRow.cpp
+++ b/Engine/Source/Editor/PropertyEditor/Private/SDetailSingleItemRow.cpp
@@ -316,29 +316,26 @@ TSharedRef<SWidget> SDetailSingleItemRow::CreateExtensionWidget(TSharedRef<SWidg
 	IDetailsViewPrivate& DetailsView = InTreeNode->GetDetailsView();
 	TSharedPtr<IDetailPropertyExtensionHandler> ExtensionHandler = DetailsView.GetExtensionHandler();
 
-	if ( ExtensionHandler.IsValid() )
+	if ( ExtensionHandler.IsValid() && InCustomization.HasPropertyNode() )
 	{
-		if ( InCustomization.HasPropertyNode() && ExtensionHandler.IsValid() )
-		{
-			TSharedPtr<IPropertyHandle> Handle = PropertyEditorHelpers::GetPropertyHandle(InCustomization.GetPropertyNode().ToSharedRef(), nullptr, nullptr);
+		TSharedPtr<IPropertyHandle> Handle = PropertyEditorHelpers::GetPropertyHandle(InCustomization.GetPropertyNode().ToSharedRef(), nullptr, nullptr);
 
-			UClass* ObjectClass = InCustomization.GetPropertyNode()->FindObjectItemParent()->GetObjectBaseClass();
-			if ( ExtensionHandler->IsPropertyExtenable(ObjectClass, *Handle) )
-			{
-				ValueWidget = SNew(SHorizontalBox)
+		UClass* ObjectClass = InCustomization.GetPropertyNode()->FindObjectItemParent()->GetObjectBaseClass();
+		if ( ExtensionHandler->IsPropertyExtenable(ObjectClass, *Handle) )
+		{
+			ValueWidget = SNew(SHorizontalBox)
 
-					+ SHorizontalBox::Slot()
-					.FillWidth(1.0f)
-					[
-						ValueWidget
-					]
+				+ SHorizontalBox::Slot()
+				.FillWidth(1.0f)
+				[
+					ValueWidget
+				]
 					
-					+ SHorizontalBox::Slot()
-					.AutoWidth()
-					[
-						ExtensionHandler->GenerateExtensionWidget(ObjectClass, Handle)
-					];
-			}
+				+ SHorizontalBox::Slot()
+				.AutoWidth()
+				[
+					ExtensionHandler->GenerateExtensionWidget(ObjectClass, Handle)
+				];
 		}
 	}
 
diff --git a/Engine/Source/Editor/PropertyEditor/Private/UserInterface/PropertyEditor/SPropertyEditorAsset.cpp b/Engine/Source/Editor/PropertyEditor/Private/UserInterface/PropertyEditor/SPropertyEditorAsset.cpp
--- a/Engine/Source/Editor/PropertyEditor/Private/UserInterface/PropertyEditor/SPropertyEditorAsset.cpp
+++ b/Engine/Source/Editor/PropertyEditor/Private/UserInterface/PropertyEditor/SPropertyEditorAsset.cpp
@@ -499,7 +499,7 @@ FPropertyAccess::Result SPropertyEditorAsset::GetValue( FObjectOrAssetData& OutV
 
 			if (CurrentObjectPath.Len() > 0 && CurrentObjectPath != TEXT("None"))
 			{
-				if( CurrentObjectPath != TEXT("None") && ( !CachedAssetData.IsValid() || CachedAssetData.ObjectPath.ToString() != CurrentObjectPath ) )
+				if( !CachedAssetData.IsValid() || CachedAssetData.ObjectPath.ToString() != CurrentObjectPath )
 				{
 					static FName AssetRegistryName("AssetRegistry");
 
diff --git a/Engine/Source/Editor/SceneOutliner/Private/SOutlinerTreeView.cpp b/Engine/Source/Editor/SceneOutliner/Private/SOutlinerTreeView.cpp
--- a/Engine/Source/Editor/SceneOutliner/Private/SOutlinerTreeView.cpp
+++ b/Engine/Source/Editor/SceneOutliner/Private/SOutlinerTreeView.cpp
@@ -330,8 +330,6 @@ namespace SceneOutliner
 		// We'll draw with the 'focused' look if we're either focused or we have a context menu summoned
 		const bool bShouldAppearFocused = HasKeyboardFocus();
 
-		const float DrawPositionY = ( AllottedGeometry.Size.Y / 2 ) - ( AllottedGeometry.Size.Y / 2 );
-
 		// Draw highlight targeting effect
 		const float TimeSinceHighlightInteraction = (float)( CurrentTime - LastHighlightInteractionTime );
 		if( TimeSinceHighlightInteraction <= HighlightTargetEffectDuration )
diff --git a/Engine/Source/Editor/UnrealEd/Classes/Settings/ContentBrowserSettings.h b/Engine/Source/Editor/UnrealEd/Classes/Settings/ContentBrowserSettings.h
--- a/Engine/Source/Editor/UnrealEd/Classes/Settings/ContentBrowserSettings.h
+++ b/Engine/Source/Editor/UnrealEd/Classes/Settings/ContentBrowserSettings.h
@@ -44,7 +44,7 @@ public:
 	/** Gets whether we are allowed to display the engine folder or not, optional flag ignoring the override */
 	bool GetDisplayEngineFolder( bool bExcludeOverride = false ) const
 	{ 
-		return ( ( bExcludeOverride ? false : OverrideDisplayEngineFolder ) | DisplayEngineFolder );
+		return ( ( bExcludeOverride ? false : OverrideDisplayEngineFolder ) || DisplayEngineFolder );
 	}
 
 	/** Sets whether we are allowed to display the developers folder or not, optional flag for setting override instead */
diff --git a/Engine/Source/Editor/UnrealEd/Private/Commandlets/ContentCommandlets.cpp b/Engine/Source/Editor/UnrealEd/Private/Commandlets/ContentCommandlets.cpp
--- a/Engine/Source/Editor/UnrealEd/Private/Commandlets/ContentCommandlets.cpp
+++ b/Engine/Source/Editor/UnrealEd/Private/Commandlets/ContentCommandlets.cpp
@@ -611,7 +611,7 @@ bool UResavePackagesCommandlet::PerformPreloadOperations( ULinkerLoad* PackageLi
 	}
 
 	// Check if this package meets the maximum requirements.
-	bool bNoLimitation = MaxResaveUE4Version == IGNORE_PACKAGE_VERSION && MaxResaveUE4Version == IGNORE_PACKAGE_VERSION && MaxResaveLicenseeUE4Version == IGNORE_PACKAGE_VERSION;
+	bool bNoLimitation = MaxResaveUE4Version == IGNORE_PACKAGE_VERSION && MaxResaveLicenseeUE4Version == IGNORE_PACKAGE_VERSION;
 	bool bAllowResave = bNoLimitation ||
 						 (MaxResaveUE4Version != IGNORE_PACKAGE_VERSION && UE4PackageVersion <= MaxResaveUE4Version) ||
 						 (MaxResaveLicenseeUE4Version != IGNORE_PACKAGE_VERSION && LicenseeUE4PackageVersion <= MaxResaveLicenseeUE4Version);
diff --git a/Engine/Source/Editor/UnrealEd/Private/Commandlets/GenerateGatherArchiveCommandlet.cpp b/Engine/Source/Editor/UnrealEd/Private/Commandlets/GenerateGatherArchiveCommandlet.cpp
--- a/Engine/Source/Editor/UnrealEd/Private/Commandlets/GenerateGatherArchiveCommandlet.cpp
+++ b/Engine/Source/Editor/UnrealEd/Private/Commandlets/GenerateGatherArchiveCommandlet.cpp
@@ -426,9 +426,6 @@ void UGenerateGatherArchiveCommandlet::ConditionSource( FLocItem& LocItem )
 {
 	if( LocItem.MetadataObj.IsValid() )
 	{
-		if( LocItem.MetadataObj.IsValid() )
-		{
-			ConditionSourceMetadata( MakeShareable( new FLocMetadataValueObject( LocItem.MetadataObj ) ) );
-		}
+		ConditionSourceMetadata( MakeShareable( new FLocMetadataValueObject( LocItem.MetadataObj ) ) );
 	}
 }
diff --git a/Engine/Source/Editor/UnrealEd/Private/DebugToolExec.cpp b/Engine/Source/Editor/UnrealEd/Private/DebugToolExec.cpp
--- a/Engine/Source/Editor/UnrealEd/Private/DebugToolExec.cpp
+++ b/Engine/Source/Editor/UnrealEd/Private/DebugToolExec.cpp
@@ -150,7 +150,7 @@ bool FDebugToolExec::Exec( UWorld* InWorld, const TCHAR* Cmd, FOutputDevice& Ar
 		if (foundObj != NULL)
 		{
 			// not allowed in the editor unless it is a PIE object as this command can have far reaching effects such as impacting serialization
-			if (!GIsEditor || ((!foundObj->IsTemplate() && (foundObj->GetOutermost()->PackageFlags & PKG_PlayInEditor))))
+			if (!GIsEditor || (!foundObj->IsTemplate() && (foundObj->GetOutermost()->PackageFlags & PKG_PlayInEditor)))
 			{
 				EditObject(foundObj, true);
 			}
@@ -175,7 +175,7 @@ bool FDebugToolExec::Exec( UWorld* InWorld, const TCHAR* Cmd, FOutputDevice& Ar
 		if (foundObj != NULL)
 		{
 			// not allowed in the editor unless it is a PIE object as this command can have far reaching effects such as impacting serialization
-			if (!GIsEditor || ((!foundObj->IsTemplate() && (foundObj->GetOutermost()->PackageFlags & PKG_PlayInEditor))))
+			if (!GIsEditor || (!foundObj->IsTemplate() && (foundObj->GetOutermost()->PackageFlags & PKG_PlayInEditor)))
 			{
 				EditObject(foundObj, false);
 			}
@@ -256,7 +256,7 @@ bool FDebugToolExec::Exec( UWorld* InWorld, const TCHAR* Cmd, FOutputDevice& Ar
 		if( Found )
 		{
 			// not allowed in the editor unless it is a PIE object as this command can have far reaching effects such as impacting serialization
-			if (!GIsEditor || ((!Found->IsTemplate() && (Found->GetOutermost()->PackageFlags & PKG_PlayInEditor))))
+			if (!GIsEditor || (!Found->IsTemplate() && (Found->GetOutermost()->PackageFlags & PKG_PlayInEditor)))
 			{
 				EditObject(Found, true);
 			}
diff --git a/Engine/Source/Editor/UnrealEd/Private/EditorLevelUtils.cpp b/Engine/Source/Editor/UnrealEd/Private/EditorLevelUtils.cpp
--- a/Engine/Source/Editor/UnrealEd/Private/EditorLevelUtils.cpp
+++ b/Engine/Source/Editor/UnrealEd/Private/EditorLevelUtils.cpp
@@ -688,7 +688,7 @@ namespace EditorLevelUtils
 				UModelComponent* CurLevelModelCmp = Level->ModelComponents[ComponentIndex];
 				if(CurLevelModelCmp)
 				{
-					if (bShouldBeVisible && CurLevelModelCmp)
+					if (bShouldBeVisible)
 					{
 						CurLevelModelCmp->RegisterComponentWithWorld(Level->OwningWorld);
 					}
diff --git a/Engine/Source/Editor/UnrealEd/Private/PackageTools.cpp b/Engine/Source/Editor/UnrealEd/Private/PackageTools.cpp
--- a/Engine/Source/Editor/UnrealEd/Private/PackageTools.cpp
+++ b/Engine/Source/Editor/UnrealEd/Private/PackageTools.cpp
@@ -436,16 +436,13 @@ namespace PackageTools
 			TArray<UObject*> FilteredObjects;
 			if ( FilteredClasses )
 			{
-				if ( FilteredClasses )
-				{
-					// Present the user with a warning that only the filtered types are being exported
-					FSuppressableWarningDialog::FSetupInfo Info( NSLOCTEXT("UnrealEd", "BulkExport_FilteredWarning", "Asset types are currently filtered within the Content Browser. Only objects of the filtered types will be exported."),
-						LOCTEXT("BulkExport_FilteredWarning_Title", "Asset Filter in Effect"), "BulkExportFilterWarning" );
-					Info.ConfirmText = NSLOCTEXT("ModalDialogs", "BulkExport_FilteredWarningConfirm", "Close");
+				// Present the user with a warning that only the filtered types are being exported
+				FSuppressableWarningDialog::FSetupInfo Info( NSLOCTEXT("UnrealEd", "BulkExport_FilteredWarning", "Asset types are currently filtered within the Content Browser. Only objects of the filtered types will be exported."),
+					LOCTEXT("BulkExport_FilteredWarning_Title", "Asset Filter in Effect"), "BulkExportFilterWarning" );
+				Info.ConfirmText = NSLOCTEXT("ModalDialogs", "BulkExport_FilteredWarningConfirm", "Close");
 
-					FSuppressableWarningDialog PromptAboutFiltering( Info );
-					PromptAboutFiltering.ShowModal();
-				}
+				FSuppressableWarningDialog PromptAboutFiltering( Info );
+				PromptAboutFiltering.ShowModal();
 				
 				for ( TArray<UObject*>::TConstIterator ObjIter(ObjectsInPackages); ObjIter; ++ObjIter )
 				{
diff --git a/Engine/Source/Editor/UnrealEd/Private/PerformanceMonitor.h b/Engine/Source/Editor/UnrealEd/Private/PerformanceMonitor.h
--- a/Engine/Source/Editor/UnrealEd/Private/PerformanceMonitor.h
+++ b/Engine/Source/Editor/UnrealEd/Private/PerformanceMonitor.h
@@ -74,7 +74,7 @@ struct FMovingAverage
 				Samples.Add(AccumulatorAverage);
 			}
 
-			NextSampleIndex = ++NextSampleIndex % SampleSize;
+			NextSampleIndex = (NextSampleIndex + 1) % SampleSize;
 			ensure(NextSampleIndex >= 0 && NextSampleIndex < SampleSize);
 
 			// calculate the average
diff --git a/Engine/Source/Editor/UnrealEd/Private/PlayLevel.cpp b/Engine/Source/Editor/UnrealEd/Private/PlayLevel.cpp
--- a/Engine/Source/Editor/UnrealEd/Private/PlayLevel.cpp
+++ b/Engine/Source/Editor/UnrealEd/Private/PlayLevel.cpp
@@ -1786,7 +1786,7 @@ void UEditorEngine::RequestEndPlayMap()
 				if (ThisContext.WorldType == EWorldType::PIE)
 				{
 					FSlatePlayInEditorInfo* const SlatePlayInEditorSession = SlatePlayInEditorMap.Find(ThisContext.ContextHandle);
-					if ((SlatePlayInEditorSession != nullptr) && (SlatePlayInEditorSession->EditorPlayer.IsValid() == true) && (SlatePlayInEditorSession->EditorPlayer.Get()->PlayerController != nullptr ) )
+					if ((SlatePlayInEditorSession != nullptr) && (SlatePlayInEditorSession->EditorPlayer.IsValid() == true) )
 					{
 						if( SlatePlayInEditorSession->EditorPlayer.Get()->PlayerController != nullptr )
 						{
diff --git a/Engine/Source/Runtime/Core/Private/HAL/IPlatformFileProfilerWrapper.cpp b/Engine/Source/Runtime/Core/Private/HAL/IPlatformFileProfilerWrapper.cpp
--- a/Engine/Source/Runtime/Core/Private/HAL/IPlatformFileProfilerWrapper.cpp
+++ b/Engine/Source/Runtime/Core/Private/HAL/IPlatformFileProfilerWrapper.cpp
@@ -77,7 +77,6 @@ bool FPlatformFileReadStats::Initialize( IPlatformFile* Inner, const TCHAR* Comm
 	LowerLevel = Inner;
 	FTicker::GetCoreTicker().AddTicker(FTickerDelegate::CreateRaw(this, &FPlatformFileReadStats::Tick), 0.f);
 
-	Timer = FPlatformTime::Seconds();
 	LifetimeReadSpeed = 0;
 	LifetimeReadSize = 0;
 	LifetimeReadCalls = 0;
diff --git a/Engine/Source/Runtime/Core/Private/Windows/WindowsPlatformMisc.cpp b/Engine/Source/Runtime/Core/Private/Windows/WindowsPlatformMisc.cpp
--- a/Engine/Source/Runtime/Core/Private/Windows/WindowsPlatformMisc.cpp
+++ b/Engine/Source/Runtime/Core/Private/Windows/WindowsPlatformMisc.cpp
@@ -64,10 +64,11 @@ int32 FWindowsOSVersionHelper::GetOSVersions( FString& out_OSVersionLabel, FStri
 	OsVersionInfo.dwOSVersionInfoSize = sizeof( OSVERSIONINFOEX );
 	out_OSVersionLabel = TEXT( "Windows (unknown version)" );
 	out_OSSubVersionLabel = FString();
+#pragma warning(push)
 #pragma warning(disable : 4996) // 'function' was declared deprecated
 	CA_SUPPRESS(28159)
 	if( GetVersionEx( (LPOSVERSIONINFO)&OsVersionInfo ) )
-#pragma warning(default : 4996)
+#pragma warning(pop)
 	{
 		bool bIsInvalidVersion = false;
 
diff --git a/Engine/Source/Runtime/Core/Private/Windows/WindowsPlatformSurvey.cpp b/Engine/Source/Runtime/Core/Private/Windows/WindowsPlatformSurvey.cpp
--- a/Engine/Source/Runtime/Core/Private/Windows/WindowsPlatformSurvey.cpp
+++ b/Engine/Source/Runtime/Core/Private/Windows/WindowsPlatformSurvey.cpp
@@ -686,10 +686,11 @@ void FWindowsPlatformSurvey::GetOSVersionLabels(const SYSTEM_INFO& SystemInfo, F
 	{
 		OSVERSIONINFOEX OsVersionInfo = {0};
 		OsVersionInfo.dwOSVersionInfoSize = sizeof( OSVERSIONINFOEX );
+#pragma warning(push)
 #pragma warning(disable : 4996) // 'function' was declared deprecated
 		CA_SUPPRESS(28159)
 		GetVersionEx( (LPOSVERSIONINFO)&OsVersionInfo );
-#pragma warning(default : 4996)
+#pragma warning(pop)
 
 		UE_LOG( LogWindows, Warning, TEXT( "FWindowsPlatformSurvey::GetOSVersionLabel() unknown Windows version info from GetVersionEx()" ) );
 		Results.ErrorCount++;
diff --git a/Engine/Source/Runtime/Core/Private/Windows/WindowsTextInputMethodSystem.cpp b/Engine/Source/Runtime/Core/Private/Windows/WindowsTextInputMethodSystem.cpp
--- a/Engine/Source/Runtime/Core/Private/Windows/WindowsTextInputMethodSystem.cpp
+++ b/Engine/Source/Runtime/Core/Private/Windows/WindowsTextInputMethodSystem.cpp
@@ -417,10 +417,11 @@ bool FWindowsTextInputMethodSystem::InitializeTSF()
 
 		TSFActivationProxy = new FTSFActivationProxy(this);
 
+#pragma warning(push)
 #pragma warning(disable : 4996) // 'function' was was declared deprecated
 		CA_SUPPRESS(28159)
 		const DWORD WindowsVersion = ::GetVersion();
-#pragma warning(default : 4996)
+#pragma warning(pop)
 
 		const DWORD WindowsMajorVersion = LOBYTE(LOWORD(WindowsVersion));
 		const DWORD WindowsMinorVersion = HIBYTE(LOWORD(WindowsVersion));
diff --git a/Engine/Source/Runtime/Core/Public/GenericPlatform/MicrosoftPlatformString.h b/Engine/Source/Runtime/Core/Public/GenericPlatform/MicrosoftPlatformString.h
--- a/Engine/Source/Runtime/Core/Public/GenericPlatform/MicrosoftPlatformString.h
+++ b/Engine/Source/Runtime/Core/Public/GenericPlatform/MicrosoftPlatformString.h
@@ -10,6 +10,7 @@
 **/
 
 #if !USE_SECURE_CRT
+#pragma warning(push)
 #pragma warning(disable : 4996) // 'function' was declared deprecated  (needed for the secure string functions)
 #endif
 
@@ -315,5 +316,5 @@ struct FMicrosoftPlatformString : public FGenericPlatformString
 };
 
 #if !USE_SECURE_CRT
-#pragma warning(default : 4996) // 'function' was was declared deprecated  (needed for the secure string functions)
+#pragma warning(pop) // 'function' was was declared deprecated  (needed for the secure string functions)
 #endif
\ No newline at end of file
diff --git a/Engine/Source/Runtime/Core/Public/Windows/WindowsPlatformCodeAnalysis.h b/Engine/Source/Runtime/Core/Public/Windows/WindowsPlatformCodeAnalysis.h
--- a/Engine/Source/Runtime/Core/Public/Windows/WindowsPlatformCodeAnalysis.h
+++ b/Engine/Source/Runtime/Core/Public/Windows/WindowsPlatformCodeAnalysis.h
@@ -93,15 +93,19 @@
 	#pragma warning(disable : 6320) // warning C6320: Exception-filter expression is the constant EXCEPTION_EXECUTE_HANDLER. This might mask exceptions that were not intended to be handled.
 
 	//PVS-Studio settings:
-	//-V::505
+	//-V::505,668,688,690
 	//-V:TRYCOMPRESSION:519,547
-	//-V:check:547
+	//-V:check(:501,547,560
 	//-V:ENABLE_TEXT_ERROR_CHECKING_RESULTS:560
-	//-V:ENABLE_LOC_TESTING:560
+	//-V:ENABLE_LOC_TESTING:560,617
 	//-V:WITH_EDITOR:560
 	//-V:UE_LOG_ACTIVE:560
 	//-V:verify:501
-	//-V:check:501
 	//-V::677,720,704
-
+	//-V:%n:609
+	//-V:UE_BUILD_SHIPPING:501
+	//-V:WITH_EDITOR:501
+	//-V:PLATFORM_:517,547
+	//-V:ensureMsg:562
+	//-V:WindowsMinorVersion:547 
 #endif
\ No newline at end of file
diff --git a/Engine/Source/Runtime/Engine/Private/AVIWriter.cpp b/Engine/Source/Runtime/Engine/Private/AVIWriter.cpp
--- a/Engine/Source/Runtime/Engine/Private/AVIWriter.cpp
+++ b/Engine/Source/Runtime/Engine/Private/AVIWriter.cpp
@@ -13,11 +13,11 @@ DEFINE_LOG_CATEGORY_STATIC(LogAVIWriter, Log, All);
 
 #include "AllowWindowsPlatformTypes.h"
 typedef TCHAR* PTCHAR;
+#pragma warning(push)
 #pragma warning(disable : 4263) // 'function' : member function does not override any base class virtual member function
 #pragma warning(disable : 4264) // 'virtual_function' : no override available for virtual member function from base 'cla
 #include <streams.h>
-#pragma warning(default : 4263) // 'function' : member function does not override any base class virtual member function
-#pragma warning(default : 4264) // 'virtual_function' : no override available for virtual member function from base 'cla
+#pragma warning(pop)
 
 #include <dshow.h>
 #include <initguid.h>
diff --git a/Engine/Source/Runtime/Engine/Private/AudioDevice.cpp b/Engine/Source/Runtime/Engine/Private/AudioDevice.cpp
--- a/Engine/Source/Runtime/Engine/Private/AudioDevice.cpp
+++ b/Engine/Source/Runtime/Engine/Private/AudioDevice.cpp
@@ -909,7 +909,7 @@ void FAudioDevice::RecurseIntoSoundClasses( USoundClass* CurrentClass, FSoundCla
 		FSoundClassProperties* Properties = SoundClasses.Find(ChildClass);
 
 		// Should never be NULL for a properly set up tree.
-		if( ChildClass && Properties )
+		if( ChildClass )
 		{
 			if (Properties)
 			{
diff --git a/Engine/Source/Runtime/Engine/Private/CaptureSource.h b/Engine/Source/Runtime/Engine/Private/CaptureSource.h
--- a/Engine/Source/Runtime/Engine/Private/CaptureSource.h
+++ b/Engine/Source/Runtime/Engine/Private/CaptureSource.h
@@ -8,13 +8,13 @@
 #define _CAPTURESOURCE_HEADER_
 
 #if PLATFORM_WINDOWS && !UE_BUILD_MINIMAL
+#pragma warning(push)
 #pragma warning(disable : 4263) // 'function' : member function does not override any base class virtual member function
 #pragma warning(disable : 4264) // 'virtual_function' : no override available for virtual member function from base 
 #include "AllowWindowsPlatformTypes.h"
 #include <streams.h>
 #include "HideWindowsPlatformTypes.h"
-#pragma warning(default : 4263) // 'function' : member function does not override any base class virtual member function
-#pragma warning(default : 4264) // 'virtual_function' : no override available for virtual member function from base
+#pragma warning(pop)
 class FCapturePin;
 
 // {9A80E195-3BBA-4821-B18B-21BB496F80F8}
diff --git a/Engine/Source/Runtime/Engine/Private/DataReplication.cpp b/Engine/Source/Runtime/Engine/Private/DataReplication.cpp
--- a/Engine/Source/Runtime/Engine/Private/DataReplication.cpp
+++ b/Engine/Source/Runtime/Engine/Private/DataReplication.cpp
@@ -643,7 +643,7 @@ bool FObjectReplicator::ReceivedBunch( FInBunch &Bunch, const FReplicationFlags&
 			}
 
 			// validate that the function is callable here
-			const bool bCanExecute = ( ( !bIsServer || RepFlags.bNetOwner ) );		// we are client or net owner
+			const bool bCanExecute = ( !bIsServer || RepFlags.bNetOwner );		// we are client or net owner
 
 			if ( bCanExecute )
 			{
diff --git a/Engine/Source/Runtime/Engine/Private/Particles/ParticleComponents.cpp b/Engine/Source/Runtime/Engine/Private/Particles/ParticleComponents.cpp
--- a/Engine/Source/Runtime/Engine/Private/Particles/ParticleComponents.cpp
+++ b/Engine/Source/Runtime/Engine/Private/Particles/ParticleComponents.cpp
@@ -1666,11 +1666,8 @@ bool UParticleEmitter::CalculateMaxActiveParticleCount()
 				UParticleModuleTypeDataBeam2* BeamTD = Cast<UParticleModuleTypeDataBeam2>(LODLevel->TypeDataModule);
 				if (BeamTD)
 				{
-					if (BeamTD)
-					{
-						bForceMaxCount = true;
-						MaxCount = BeamTD->MaxBeamCount + 2;
-					}
+					bForceMaxCount = true;
+					MaxCount = BeamTD->MaxBeamCount + 2;
 				}
 			}
 
diff --git a/Engine/Source/Runtime/GameMenuBuilder/Private/SGameMenuPageWidget.cpp b/Engine/Source/Runtime/GameMenuBuilder/Private/SGameMenuPageWidget.cpp
--- a/Engine/Source/Runtime/GameMenuBuilder/Private/SGameMenuPageWidget.cpp
+++ b/Engine/Source/Runtime/GameMenuBuilder/Private/SGameMenuPageWidget.cpp
@@ -281,7 +281,7 @@ EVisibility SGameMenuPageWidget::GetSlateVisibility() const
 
 EVisibility SGameMenuPageWidget::GetMenuTitleVisibility() const
 {
-	return CurrentMenuTitle.IsEmpty() || bMenuHidden == true ? EVisibility::Collapsed : EVisibility::Visible;
+	return (CurrentMenuTitle.IsEmpty() || bMenuHidden == true) ? EVisibility::Collapsed : EVisibility::Visible;
 }
 
 FText SGameMenuPageWidget::GetMenuTitle() const 
diff --git a/Engine/Source/Runtime/Landscape/Private/LandscapeEditInterface.cpp b/Engine/Source/Runtime/Landscape/Private/LandscapeEditInterface.cpp
--- a/Engine/Source/Runtime/Landscape/Private/LandscapeEditInterface.cpp
+++ b/Engine/Source/Runtime/Landscape/Private/LandscapeEditInterface.cpp
@@ -4759,11 +4759,8 @@ FLandscapeTextureDataInfo::~FLandscapeTextureDataInfo()
 	{
 		if( MipInfo[i].MipData )
 		{
-			if( MipInfo[i].MipData )
-			{
-				Texture->Source.UnlockMip(i);
-				MipInfo[i].MipData = NULL;
-			}
+			Texture->Source.UnlockMip(i);
+			MipInfo[i].MipData = NULL;
 		}
 	}
 	Texture->ClearFlags(RF_Transactional);
diff --git a/Engine/Source/Runtime/Landscape/Private/Materials/MaterialExpressionLandscapeLayerBlend.cpp b/Engine/Source/Runtime/Landscape/Private/Materials/MaterialExpressionLandscapeLayerBlend.cpp
--- a/Engine/Source/Runtime/Landscape/Private/Materials/MaterialExpressionLandscapeLayerBlend.cpp
+++ b/Engine/Source/Runtime/Landscape/Private/Materials/MaterialExpressionLandscapeLayerBlend.cpp
@@ -238,7 +238,7 @@ int32 UMaterialExpressionLandscapeLayerBlend::Compile(class FMaterialCompiler* C
 				bool bHeightIsMaterialAttributes = Layers[LayerIdx].HeightInput.Expression->IsResultMaterialAttributes(Layers[LayerIdx].HeightInput.OutputIndex);
 				if (bHeightIsMaterialAttributes)
 				{
-					Compiler->Errorf(TEXT("Height input (%s) does not accept MaterialAttributes"), Layers[LayerIdx].LayerName);
+					Compiler->Errorf(TEXT("Height input (%s) does not accept MaterialAttributes"), *(Layers[LayerIdx].LayerName.ToString()));
 				}
 			}
 			if (Layers[LayerIdx].LayerInput.Expression)
diff --git a/Engine/Source/Runtime/SlateRHIRenderer/Private/SlateRHIRenderer.cpp b/Engine/Source/Runtime/SlateRHIRenderer/Private/SlateRHIRenderer.cpp
--- a/Engine/Source/Runtime/SlateRHIRenderer/Private/SlateRHIRenderer.cpp
+++ b/Engine/Source/Runtime/SlateRHIRenderer/Private/SlateRHIRenderer.cpp
@@ -747,10 +747,7 @@ FIntRect FSlateRHIRenderer::SetupVirtualScreenBuffer(const bool bPrimaryWorkArea
 		if( CrashTrackerResource != nullptr )
 		{
 			// Size has changed, so clear out our old resource and create a new one
-			if (CrashTrackerResource != nullptr)
-			{
-				BeginReleaseResource(CrashTrackerResource);
-			}
+			BeginReleaseResource(CrashTrackerResource);
 
 			FlushRenderingCommands();
 	
